Flavor of ADPCM encode was recently refactored and exposed a bug where noninterleaved encoding was not reporting the correct number of samples per channel.

diff --git a/Engine/Source/Developer/AudioFormatADPCM/Private/AudioFormatADPCM.cpp b/Engine/Source/Developer/AudioFormatADPCM/Private/AudioFormatADPCM.cpp
--- a/Engine/Source/Developer/AudioFormatADPCM/Private/AudioFormatADPCM.cpp
+++ b/Engine/Source/Developer/AudioFormatADPCM/Private/AudioFormatADPCM.cpp
@@ -347,8 +347,7 @@ namespace ADPCM
 		const int32 SourceSampleStride = 1;
 
 		// Input source samples are 2-bytes
-		const int32 SourceNumSamples = QualityInfo.SampleDataSize / 2;
-		const int32 SourceNumSamplesPerChannel = SourceNumSamples / QualityInfo.NumChannels;
+		const int32 SourceNumSamplesPerChannel = QualityInfo.SampleDataSize / 2;
 
 		// Output samples are 4-bits
 		const int32 CompressedNumSamplesPerByte = 2;
@@ -374,7 +373,11 @@ namespace ADPCM
 
 			for (int32 BlockIndex = 0; BlockIndex < NumBlocksPerChannel; ++BlockIndex)
 			{
-				EncodeBlock(ChannelPCMSamples + SourceSampleOffset, SourceSampleStride, SourceNumSamples - SourceSampleOffset, BlockSize, EncodedADPCMChannelData + DestDataOffset);
+				const int16* InputPCMSampleOffset = ChannelPCMSamples + SourceSampleOffset;
+				const int32 NumBlockSamples = (SourceNumSamplesPerChannel * QualityInfo.NumChannels) - SourceSampleOffset;
+				uint8* NewlyEncodedADPCMData = EncodedADPCMChannelData + DestDataOffset;
+
+				EncodeBlock(InputPCMSampleOffset, SourceSampleStride, NumBlockSamples, BlockSize, NewlyEncodedADPCMData);
 
 				SourceSampleOffset += CompressedSamplesPerBlock * SourceSampleStride;
 				DestDataOffset += BlockSize;
@@ -430,7 +433,7 @@ class FAudioFormatADPCM : public IAudioFormat
 
 		while (CurrentByte < Bytes)
 		{
-			for (auto SrcBuffer : SrcBuffers)
+			for (const TArray<uint8>& SrcBuffer : SrcBuffers)
 			{
 				// our data is int16 
 				InterleavedBuffer.Push(SrcBuffer[CurrentByte]);
