[NO NON-PII COMMIT MESSAGE]

diff --git a/Engine/Source/Developer/AudioFormatADPCM/Private/AudioFormatADPCM.cpp b/Engine/Source/Developer/AudioFormatADPCM/Private/AudioFormatADPCM.cpp
--- a/Engine/Source/Developer/AudioFormatADPCM/Private/AudioFormatADPCM.cpp
+++ b/Engine/Source/Developer/AudioFormatADPCM/Private/AudioFormatADPCM.cpp
@@ -347,7 +347,7 @@ namespace ADPCM
 		const int32 SourceSampleStride = 1;
 
 		// Input source samples are 2-bytes
-		const int32 SourceNumSamplesPerChannel = QualityInfo.SampleDataSize / 2;
+		const int32 SourceNumSamplesPerChannel = (InputPCMData.Num() > 0) ? (InputPCMData[0].Num() / 2) : 0;
 
 		// Output samples are 4-bits
 		const int32 CompressedNumSamplesPerByte = 2;
diff --git a/Engine/Source/Runtime/Engine/Private/AudioDerivedData.cpp b/Engine/Source/Runtime/Engine/Private/AudioDerivedData.cpp
--- a/Engine/Source/Runtime/Engine/Private/AudioDerivedData.cpp
+++ b/Engine/Source/Runtime/Engine/Private/AudioDerivedData.cpp
@@ -1262,7 +1262,7 @@ static void CookSurroundWave( USoundWave* SoundWave, FName FormatName, const IAu
 				}
 
 				WaveSampleRate = SampleRateOverride;
-				SampleDataSize = SizeOfZerothChannel * ChannelCount;
+				SampleDataSize = SizeOfZerothChannel;
 			}
 
 			UE_LOG(LogAudioDerivedData, Display, TEXT("Cooking %d channels for: %s"), ChannelCount, *SoundWave->GetFullName());
