Add Insight trace to audio encoders

#rnx

diff --git a/Engine/Source/Developer/AudioFormatADPCM/Private/AudioFormatADPCM.cpp b/Engine/Source/Developer/AudioFormatADPCM/Private/AudioFormatADPCM.cpp
--- a/Engine/Source/Developer/AudioFormatADPCM/Private/AudioFormatADPCM.cpp
+++ b/Engine/Source/Developer/AudioFormatADPCM/Private/AudioFormatADPCM.cpp
@@ -463,6 +463,7 @@ public:
 
 	virtual bool Cook(FName Format, const TArray<uint8>& SrcBuffer, FSoundQualityInfo& QualityInfo, TArray<uint8>& CompressedDataStore) const
 	{
+		TRACE_CPUPROFILER_EVENT_SCOPE(FAudioFormatADPCM::Cook);
 		check(Format == Audio::NAME_ADPCM || Format == Audio::NAME_PCM);
 
 		if (QualityInfo.Quality == 100)
@@ -479,6 +480,8 @@ public:
 
 	virtual bool CookSurround(FName Format, const TArray<TArray<uint8> >& SrcBuffers, FSoundQualityInfo& QualityInfo, TArray<uint8>& CompressedDataStore) const
 	{
+		TRACE_CPUPROFILER_EVENT_SCOPE(FAudioFormatADPCM::CookSurround);
+
 		// Ensure the right format
 		check(Format == Audio::NAME_ADPCM || Format == Audio::NAME_PCM);
 		// Ensure at least two channel
diff --git a/Engine/Source/Developer/AudioFormatBink/Private/AudioFormatBink.cpp b/Engine/Source/Developer/AudioFormatBink/Private/AudioFormatBink.cpp
--- a/Engine/Source/Developer/AudioFormatBink/Private/AudioFormatBink.cpp
+++ b/Engine/Source/Developer/AudioFormatBink/Private/AudioFormatBink.cpp
@@ -44,6 +44,7 @@ public:
 
 	virtual bool Cook(FName InFormat, const TArray<uint8>& InSrcBuffer, FSoundQualityInfo& InQualityInfo, TArray<uint8>& OutCompressedDataStore) const override
 	{
+		TRACE_CPUPROFILER_EVENT_SCOPE(FAudioFormatBink::Cook);
 		check(InFormat == NAME_BINKA);
 
 		// Bink goes from 0 (best) to 9 (worst), but is basically unusable below like 4,
@@ -63,6 +64,7 @@ public:
 
 	virtual bool CookSurround(FName InFormat, const TArray<TArray<uint8> >& InSrcBuffers, FSoundQualityInfo& InQualityInfo, TArray<uint8>& OutCompressedDataStore) const override
 	{
+		TRACE_CPUPROFILER_EVENT_SCOPE(FAudioFormatBink::CookSurround);
 		check(InFormat == NAME_BINKA);
 
 		//
diff --git a/Engine/Source/Developer/AudioFormatOgg/Private/AudioFormatOgg.cpp b/Engine/Source/Developer/AudioFormatOgg/Private/AudioFormatOgg.cpp
--- a/Engine/Source/Developer/AudioFormatOgg/Private/AudioFormatOgg.cpp
+++ b/Engine/Source/Developer/AudioFormatOgg/Private/AudioFormatOgg.cpp
@@ -55,6 +55,7 @@ public:
 
 	virtual bool Cook(FName Format, const TArray<uint8>& SrcBuffer, FSoundQualityInfo& QualityInfo, TArray<uint8>& CompressedDataStore) const override
 	{
+		TRACE_CPUPROFILER_EVENT_SCOPE(FAudioFormatOgg::Cook);
 		check(Format == NAME_OGG);
 #if WITH_OGGVORBIS
 		{
@@ -206,6 +207,7 @@ public:
 
 	virtual bool CookSurround(FName Format, const TArray<TArray<uint8> >& SrcBuffers, FSoundQualityInfo& QualityInfo, TArray<uint8>& CompressedDataStore) const override
 	{
+		TRACE_CPUPROFILER_EVENT_SCOPE(FAudioFormatOgg::CookSurround);
 		check(Format == NAME_OGG);
 #if WITH_OGGVORBIS
 		{
diff --git a/Engine/Source/Developer/AudioFormatOpus/Private/AudioFormatOpus.cpp b/Engine/Source/Developer/AudioFormatOpus/Private/AudioFormatOpus.cpp
--- a/Engine/Source/Developer/AudioFormatOpus/Private/AudioFormatOpus.cpp
+++ b/Engine/Source/Developer/AudioFormatOpus/Private/AudioFormatOpus.cpp
@@ -53,6 +53,7 @@ public:
 
 	virtual bool Cook(FName Format, const TArray<uint8>& SrcBuffer, FSoundQualityInfo& QualityInfo, TArray<uint8>& CompressedDataStore) const override
 	{
+		TRACE_CPUPROFILER_EVENT_SCOPE(FAudioFormatOpus::Cook);
 		check(Format == NAME_OPUS);
 
 		// Get best compatible sample rate
@@ -153,6 +154,7 @@ public:
 
 	virtual bool CookSurround(FName Format, const TArray<TArray<uint8> >& SrcBuffers, FSoundQualityInfo& QualityInfo, TArray<uint8>& CompressedDataStore) const override
 	{
+		TRACE_CPUPROFILER_EVENT_SCOPE(FAudioFormatOpus::CookSurround);
 		check(Format == NAME_OPUS);
 
 		// Get best compatible sample rate
