Added protection against crash when importing corrupt audio file and setting codec to PCM.

diff --git a/Engine/Source/Developer/AudioFormatADPCM/Private/AudioFormatADPCM.cpp b/Engine/Source/Developer/AudioFormatADPCM/Private/AudioFormatADPCM.cpp
--- a/Engine/Source/Developer/AudioFormatADPCM/Private/AudioFormatADPCM.cpp
+++ b/Engine/Source/Developer/AudioFormatADPCM/Private/AudioFormatADPCM.cpp
@@ -602,7 +602,7 @@ public:
 
 			int32 DataLeftInCurChunk = InitialMaxChunkSize - HeaderSize;
 
-			while (SrcSize > 0)
+			while (SrcSize > 0 && DataLeftInCurChunk > 0)
 			{
 				// Calculate how many frames can fit in whats left of the current chunk
 				DataLeftInCurChunk = FMath::Min(DataLeftInCurChunk, SrcSize);
