Fix for PCM not playing under some circumstances

diff --git a/Engine/Source/Developer/AudioFormatADPCM/Private/AudioFormatADPCM.cpp b/Engine/Source/Developer/AudioFormatADPCM/Private/AudioFormatADPCM.cpp
--- a/Engine/Source/Developer/AudioFormatADPCM/Private/AudioFormatADPCM.cpp
+++ b/Engine/Source/Developer/AudioFormatADPCM/Private/AudioFormatADPCM.cpp
@@ -416,7 +416,7 @@ class FAudioFormatADPCM : public IAudioFormat
 	enum
 	{
 		/** Version for ADPCM format, this becomes part of the DDC key. */
-		UE_AUDIO_ADPCM_VER = 5,
+		UE_AUDIO_ADPCM_VER = 6,
 	};
 
 	void InterleaveBuffers(const TArray<TArray<uint8> >& SrcBuffers, TArray<uint8> & InterleavedBuffer) const
@@ -600,7 +600,8 @@ public:
 			SrcSize -= HeaderSize;
 			SrcData = WaveInfo.SampleDataStart;
 
-			int32 DataLeftInCurChunk = InitialMaxChunkSize - HeaderSize;
+			const int32 Chunk0Remaining = InitialMaxChunkSize - HeaderSize;
+			int32 DataLeftInCurChunk = Chunk0Remaining <= 0 ? MaxChunkSize : Chunk0Remaining;
 
 			while (SrcSize > 0 && DataLeftInCurChunk > 0)
 			{
