Fix for crash/OOM on importing truncated .wav files if using the PCM codec.

diff --git a/Engine/Source/Developer/AudioFormatADPCM/Private/AudioFormatADPCM.cpp b/Engine/Source/Developer/AudioFormatADPCM/Private/AudioFormatADPCM.cpp
--- a/Engine/Source/Developer/AudioFormatADPCM/Private/AudioFormatADPCM.cpp
+++ b/Engine/Source/Developer/AudioFormatADPCM/Private/AudioFormatADPCM.cpp
@@ -616,6 +616,13 @@ public:
 				SrcSize -= SizeOfNewChunk;
 				SrcData += SizeOfNewChunk;
 
+				// Handle bad data. If there's less than a single frame remaining
+				// bail, otherwise we can get stuck.
+				if (SrcSize < FrameSize)
+				{
+					break;
+				}
+
 				if (SrcSize > 0)
 				{
 					DataLeftInCurChunk = MaxChunkSize;
