Psuedo Rollback fixes around handling truncated files, as they caused bugs with seeking.

diff --git a/Engine/Source/Developer/AudioFormatADPCM/Private/AudioFormatADPCM.cpp b/Engine/Source/Developer/AudioFormatADPCM/Private/AudioFormatADPCM.cpp
--- a/Engine/Source/Developer/AudioFormatADPCM/Private/AudioFormatADPCM.cpp
+++ b/Engine/Source/Developer/AudioFormatADPCM/Private/AudioFormatADPCM.cpp
@@ -416,7 +416,7 @@ class FAudioFormatADPCM : public IAudioFormat
 	enum
 	{
 		/** Version for ADPCM format, this becomes part of the DDC key. */
-		UE_AUDIO_ADPCM_VER = 6,
+		UE_AUDIO_ADPCM_VER = 7,
 	};
 
 	void InterleaveBuffers(const TArray<TArray<uint8> >& SrcBuffers, TArray<uint8> & InterleavedBuffer) const
@@ -600,10 +600,9 @@ public:
 			SrcSize -= HeaderSize;
 			SrcData = WaveInfo.SampleDataStart;
 
-			const int32 Chunk0Remaining = InitialMaxChunkSize - HeaderSize;
-			int32 DataLeftInCurChunk = Chunk0Remaining <= 0 ? MaxChunkSize : Chunk0Remaining;
+			int32 DataLeftInCurChunk = InitialMaxChunkSize - HeaderSize;
 
-			while (SrcSize > 0 && DataLeftInCurChunk > 0)
+			while (SrcSize > 0)
 			{
 				// Calculate how many frames can fit in whats left of the current chunk
 				DataLeftInCurChunk = FMath::Min(DataLeftInCurChunk, SrcSize);
@@ -616,13 +615,7 @@ public:
 				SrcSize -= SizeOfNewChunk;
 				SrcData += SizeOfNewChunk;
 
-				// Handle bad data. If there's less than a single frame remaining
-				// bail, otherwise we can get stuck.
-				if (SrcSize < FrameSize)
-				{
-					break;
-				}
-
+				
 				if (SrcSize > 0)
 				{
 					DataLeftInCurChunk = MaxChunkSize;
